Automatically select the edited SoundClass when opening the Sound Class editor rather than displaying the root values.

diff --git a/Engine/Source/Editor/SoundClassEditor/Private/SoundClassEditor.cpp b/Engine/Source/Editor/SoundClassEditor/Private/SoundClassEditor.cpp
--- a/Engine/Source/Editor/SoundClassEditor/Private/SoundClassEditor.cpp
+++ b/Engine/Source/Editor/SoundClassEditor/Private/SoundClassEditor.cpp
@@ -120,13 +120,18 @@ void FSoundClassEditor::InitSoundClassEditor( const EToolkitMode::Type Mode, con
 	AddMenuExtender(SoundClassEditorModule->GetMenuExtensibilityManager()->GetAllExtenders(GetToolkitCommands(), GetEditingObjects()));
 	AddToolbarExtender(SoundClassEditorModule->GetToolBarExtensibilityManager()->GetAllExtenders(GetToolkitCommands(), GetEditingObjects()));
 
-	// @todo toolkit world centric editing
-	/*if( IsWorldCentricAssetEditor() )
+	GraphEditor->SelectAllNodes();
+	for (UObject* SelectedNode : GraphEditor->GetSelectedNodes())
 	{
-		SpawnToolkitTab(GetToolbarTabId(), FString(), EToolkitTabSpot::ToolBar);
-		SpawnToolkitTab(GraphCanvasTabId, FString(), EToolkitTabSpot::Document);
-		SpawnToolkitTab(PropertiesTabId, FString(), EToolkitTabSpot::Details);
-	}*/
+		USoundClassGraphNode* GraphNode = CastChecked<USoundClassGraphNode>(SelectedNode);
+		if (GraphNode->SoundClass == ObjectToEdit)
+		{
+			GraphEditor->ClearSelectionSet();
+			GraphEditor->SetNodeSelection(GraphNode, true);
+			DetailsView->SetObject(ObjectToEdit);
+			break;
+		}
+	}
 }
 
 FSoundClassEditor::~FSoundClassEditor()
