Fix Win64 editor link error.

- This dependency may be superficial, needs to be investigated.

diff --git a/Engine/Source/Editor/SoundClassEditor/SoundClassEditor.Build.cs b/Engine/Source/Editor/SoundClassEditor/SoundClassEditor.Build.cs
--- a/Engine/Source/Editor/SoundClassEditor/SoundClassEditor.Build.cs
+++ b/Engine/Source/Editor/SoundClassEditor/SoundClassEditor.Build.cs
@@ -20,6 +20,7 @@ public class SoundClassEditor : ModuleRules
 				"Engine",
 				"Slate",
 				"SlateCore",
+				"InputCore",
                 "EditorStyle",
 				"UnrealEd",
 				"GraphEditor"
